Add taskbar icons to UnrealMultiUserSlateServer

diff --git a/Engine/Source/Programs/UnrealMultiUserSlateServer/Resources/Windows/UnrealMultiUserSlateServer.rc b/Engine/Source/Programs/UnrealMultiUserSlateServer/Resources/Windows/UnrealMultiUserSlateServer.rc
new file mode 100644
--- /dev/null
+++ b/Engine/Source/Programs/UnrealMultiUserSlateServer/Resources/Windows/UnrealMultiUserSlateServer.rc
@@ -0,0 +1,89 @@
+// Copyright Epic Games, Inc. All Rights Reserved.
+
+#define APSTUDIO_READONLY_SYMBOLS
+/////////////////////////////////////////////////////////////////////////////
+//
+// Generated from the TEXTINCLUDE 2 resource.
+//
+#include <windows.h>
+#include "Runtime/Launch/Resources/Version.h"
+#include "Runtime/Launch/Resources/Windows/resource.h"
+#include "Developer/TraceInsights/Public/Insights/Version.h"
+
+#define UNREAL_MULTI_USER_SLATE_SERVER_VERSION_MAYOR 1
+#define UNREAL_MULTI_USER_SLATE_SERVER_VERSION_MINOR 0
+#define UNREAL_MULTI_USER_SLATE_SERVER_VERSION_PATCH 0
+
+#define UNREAL_MULTI_USER_SLATE_SERVER_VERSION_STRING "v1.01"
+
+#define VER_FILEVERSION         UNREAL_MULTI_USER_SLATE_SERVER_VERSION_MAYOR,UNREAL_MULTI_USER_SLATE_SERVER_VERSION_MINOR,UNREAL_MULTI_USER_SLATE_SERVER_VERSION_PATCH,0
+#define VER_FILEVERSION_STR     UNREAL_MULTI_USER_SLATE_SERVER_VERSION_STRING "\0"
+
+#define VER_PRODUCTVERSION      ENGINE_MAJOR_VERSION,ENGINE_MINOR_VERSION,ENGINE_PATCH_VERSION,0
+#define VER_PRODUCTVERSION_STR  ENGINE_VERSION_STRING "\0"
+
+#ifdef _DEBUG
+#define VER_FILEFLAGS           0x1L // VS_FF_DEBUG
+#else
+#define VER_FILEFLAGS           0x0L
+#endif
+
+#define IDI_UNREAL_MULTI_USER_SLATE_SERVER            IDICON_UEGame // 123
+
+/////////////////////////////////////////////////////////////////////////////
+#undef APSTUDIO_READONLY_SYMBOLS
+
+/////////////////////////////////////////////////////////////////////////////
+// English (U.S.) resources
+
+#if !defined(AFX_RESOURCE_DLL) || defined(AFX_TARG_ENU)
+#ifdef _WIN32
+LANGUAGE LANG_ENGLISH, SUBLANG_ENGLISH_US
+#pragma code_page(1252)
+#endif //_WIN32
+
+/////////////////////////////////////////////////////////////////////////////
+//
+// Version
+//
+
+VS_VERSION_INFO VERSIONINFO
+FILEVERSION     VER_FILEVERSION
+PRODUCTVERSION  VER_PRODUCTVERSION
+FILEFLAGSMASK   0x17L
+FILEFLAGS       VER_FILEFLAGS
+FILEOS          0x4L
+FILETYPE        0x2L
+FILESUBTYPE     0x0L
+BEGIN
+    BLOCK "StringFileInfo"
+    BEGIN
+        BLOCK "040904b0"
+        BEGIN
+            VALUE "CompanyName",      EPIC_COMPANY_NAME
+            VALUE "LegalCopyright",   EPIC_COPYRIGHT_STRING
+            VALUE "ProductName",      EPIC_PRODUCT_NAME
+            VALUE "ProductVersion",   VER_PRODUCTVERSION_STR
+            VALUE "FileDescription",  "Unreal Insights"
+            VALUE "FileVersion",      VER_FILEVERSION_STR
+            VALUE "InternalName",     "UnrealMultiUserSlateServer"
+            VALUE "OriginalFilename", "UnrealMultiUserSlateServer.exe"
+        END
+    END
+    BLOCK "VarFileInfo"
+    BEGIN
+        VALUE "Translation", 0x409, 1200
+    END
+END
+
+/////////////////////////////////////////////////////////////////////////////
+//
+// Icon
+//
+
+// Icon with lowest ID value placed first to ensure application icon
+// remains consistent on all systems.
+IDI_UNREAL_MULTI_USER_SLATE_SERVER    ICON    "UnrealMultiUserSlateServer.ico"
+
+#endif    // English (U.S.) resources
+/////////////////////////////////////////////////////////////////////////////
diff --git a/Engine/Source/Programs/UnrealMultiUserSlateServer/Resources/Windows/make_ico.py b/Engine/Source/Programs/UnrealMultiUserSlateServer/Resources/Windows/make_ico.py
new file mode 100644
--- /dev/null
+++ b/Engine/Source/Programs/UnrealMultiUserSlateServer/Resources/Windows/make_ico.py
@@ -0,0 +1,15 @@
+# Copyright Epic Games, Inc. All Rights Reserved.
+
+from PIL import Image
+
+def MakeIcon(folder, output):
+    sizes = [256, 48, 32, 24, 16]
+    images = []
+    for size in sizes:
+        print(f"Loading {folder}/icon_{size}x{size}.png")
+        image = Image.open(f"{folder}/icon_{size}x{size}.png")
+        images.append(image)
+    print(f"Creating {output}")
+    images[0].save(f"{output}", append_images=images)
+
+MakeIcon("UnrealMultiUserSlateServer", "UnrealMultiUserSlateServer.ico")
