Add ShortName to module generating paths longer than 200 characters

diff --git a/Engine/Source/Programs/UnrealMultiUserSlateServer/UnrealMultiUserSlateServer.Build.cs b/Engine/Source/Programs/UnrealMultiUserSlateServer/UnrealMultiUserSlateServer.Build.cs
--- a/Engine/Source/Programs/UnrealMultiUserSlateServer/UnrealMultiUserSlateServer.Build.cs
+++ b/Engine/Source/Programs/UnrealMultiUserSlateServer/UnrealMultiUserSlateServer.Build.cs
@@ -43,5 +43,7 @@ public class UnrealMultiUserSlateServer : ModuleRules
 				}
 			);
 		}
+
+		ShortName = "UEMultiUsrSltServ";
 	}
 }
