- Updated icon files in some Epic programs for Modern Xcode, and point to existing plist files

diff --git a/Engine/Programs/UnrealInsights/Config/DefaultEngine.ini b/Engine/Programs/UnrealInsights/Config/DefaultEngine.ini
--- a/Engine/Programs/UnrealInsights/Config/DefaultEngine.ini
+++ b/Engine/Programs/UnrealInsights/Config/DefaultEngine.ini
@@ -6,3 +6,6 @@
 +ProgramEnabledPlugins=NetworkPredictionInsights
 +ProgramEnabledPlugins=VisualStudioSourceCodeAccess
 +ProgramEnabledPlugins=XcodeSourceCodeAccess
+
+[/Script/MacTargetPlatform.XcodeProjectSettings]
+TemplateMacPlist=(FilePath="/Engine/Source/Programs/UnrealInsights/Resources/Mac/Info.plist")
diff --git a/Engine/Programs/UnrealMultiUserServer/Config/DefaultEngine.ini b/Engine/Programs/UnrealMultiUserServer/Config/DefaultEngine.ini
--- a/Engine/Programs/UnrealMultiUserServer/Config/DefaultEngine.ini
+++ b/Engine/Programs/UnrealMultiUserServer/Config/DefaultEngine.ini
@@ -5,3 +5,6 @@ GameNameShort=MUS
 +ProgramEnabledPlugins="UdpMessaging"
 +ProgramEnabledPlugins="QuicMessaging"
 +ProgramEnabledPlugins="ConcertSyncServer"
+
+[/Script/MacTargetPlatform.XcodeProjectSettings]
+TemplateMacPlist=(FilePath="/Engine/Source/Programs/UnrealMultiUserServer/Resources/Mac/Info.plist")
diff --git a/Engine/Programs/UnrealMultiUserSlateServer/Config/DefaultEngine.ini b/Engine/Programs/UnrealMultiUserSlateServer/Config/DefaultEngine.ini
--- a/Engine/Programs/UnrealMultiUserSlateServer/Config/DefaultEngine.ini
+++ b/Engine/Programs/UnrealMultiUserSlateServer/Config/DefaultEngine.ini
@@ -6,3 +6,6 @@ GameNameShort=MUSS
 +ProgramEnabledPlugins="QuicMessaging"
 +ProgramEnabledPlugins="ConcertSyncServer"
 +ProgramEnabledPlugins="MultiUserServer"
+
+[/Script/MacTargetPlatform.XcodeProjectSettings]
+TemplateMacPlist=(FilePath="/Engine/Source/Programs/UnrealMultiUserSlateServer/Resources/Mac/Info.plist")
diff --git a/Engine/Source/Programs/UnrealInsights/Resources/Mac/Assets.xcassets/AppIcon.appiconset/Contents.json b/Engine/Source/Programs/UnrealInsights/Resources/Mac/Assets.xcassets/AppIcon.appiconset/Contents.json
new file mode 100644
--- /dev/null
+++ b/Engine/Source/Programs/UnrealInsights/Resources/Mac/Assets.xcassets/AppIcon.appiconset/Contents.json
@@ -0,0 +1,68 @@
+{
+  "images" : [
+    {
+      "filename" : "icon_16x16.png",
+      "idiom" : "mac",
+      "scale" : "1x",
+      "size" : "16x16"
+    },
+    {
+      "filename" : "[REDACTED_EMAIL]",
+      "idiom" : "mac",
+      "scale" : "2x",
+      "size" : "16x16"
+    },
+    {
+      "filename" : "icon_32x32.png",
+      "idiom" : "mac",
+      "scale" : "1x",
+      "size" : "32x32"
+    },
+    {
+      "filename" : "[REDACTED_EMAIL]",
+      "idiom" : "mac",
+      "scale" : "2x",
+      "size" : "32x32"
+    },
+    {
+      "filename" : "icon_128x128.png",
+      "idiom" : "mac",
+      "scale" : "1x",
+      "size" : "128x128"
+    },
+    {
+      "filename" : "[REDACTED_EMAIL]",
+      "idiom" : "mac",
+      "scale" : "2x",
+      "size" : "128x128"
+    },
+    {
+      "filename" : "icon_256x256.png",
+      "idiom" : "mac",
+      "scale" : "1x",
+      "size" : "256x256"
+    },
+    {
+      "filename" : "[REDACTED_EMAIL]",
+      "idiom" : "mac",
+      "scale" : "2x",
+      "size" : "256x256"
+    },
+    {
+      "filename" : "icon_512x512.png",
+      "idiom" : "mac",
+      "scale" : "1x",
+      "size" : "512x512"
+    },
+    {
+      "filename" : "[REDACTED_EMAIL]",
+      "idiom" : "mac",
+      "scale" : "2x",
+      "size" : "512x512"
+    }
+  ],
+  "info" : {
+    "author" : "xcode",
+    "version" : 1
+  }
+}
diff --git a/Engine/Source/Programs/UnrealInsights/Resources/Mac/Assets.xcassets/Contents.json b/Engine/Source/Programs/UnrealInsights/Resources/Mac/Assets.xcassets/Contents.json
new file mode 100644
--- /dev/null
+++ b/Engine/Source/Programs/UnrealInsights/Resources/Mac/Assets.xcassets/Contents.json
@@ -0,0 +1,6 @@
+{
+  "info" : {
+    "author" : "xcode",
+    "version" : 1
+  }
+}
diff --git a/Engine/Source/Programs/UnrealInsights/Resources/Mac/make_icns.sh b/Engine/Source/Programs/UnrealInsights/Resources/Mac/make_icns.sh
deleted file mode 100644
diff --git a/Engine/Source/Programs/UnrealMultiUserServer/Resources/Mac/Assets.xcassets/AppIcon.appiconset/Contents.json b/Engine/Source/Programs/UnrealMultiUserServer/Resources/Mac/Assets.xcassets/AppIcon.appiconset/Contents.json
new file mode 100644
--- /dev/null
+++ b/Engine/Source/Programs/UnrealMultiUserServer/Resources/Mac/Assets.xcassets/AppIcon.appiconset/Contents.json
@@ -0,0 +1,68 @@
+{
+  "images" : [
+    {
+      "filename" : "icon_16x16.png",
+      "idiom" : "mac",
+      "scale" : "1x",
+      "size" : "16x16"
+    },
+    {
+      "filename" : "[REDACTED_EMAIL]",
+      "idiom" : "mac",
+      "scale" : "2x",
+      "size" : "16x16"
+    },
+    {
+      "filename" : "icon_32x32.png",
+      "idiom" : "mac",
+      "scale" : "1x",
+      "size" : "32x32"
+    },
+    {
+      "filename" : "[REDACTED_EMAIL]",
+      "idiom" : "mac",
+      "scale" : "2x",
+      "size" : "32x32"
+    },
+    {
+      "filename" : "icon_128x128.png",
+      "idiom" : "mac",
+      "scale" : "1x",
+      "size" : "128x128"
+    },
+    {
+      "filename" : "[REDACTED_EMAIL]",
+      "idiom" : "mac",
+      "scale" : "2x",
+      "size" : "128x128"
+    },
+    {
+      "filename" : "icon_256x256.png",
+      "idiom" : "mac",
+      "scale" : "1x",
+      "size" : "256x256"
+    },
+    {
+      "filename" : "[REDACTED_EMAIL]",
+      "idiom" : "mac",
+      "scale" : "2x",
+      "size" : "256x256"
+    },
+    {
+      "filename" : "icon_512x512.png",
+      "idiom" : "mac",
+      "scale" : "1x",
+      "size" : "512x512"
+    },
+    {
+      "filename" : "[REDACTED_EMAIL]",
+      "idiom" : "mac",
+      "scale" : "2x",
+      "size" : "512x512"
+    }
+  ],
+  "info" : {
+    "author" : "xcode",
+    "version" : 1
+  }
+}
diff --git a/Engine/Source/Programs/UnrealMultiUserServer/Resources/Mac/Assets.xcassets/Contents.json b/Engine/Source/Programs/UnrealMultiUserServer/Resources/Mac/Assets.xcassets/Contents.json
new file mode 100644
--- /dev/null
+++ b/Engine/Source/Programs/UnrealMultiUserServer/Resources/Mac/Assets.xcassets/Contents.json
@@ -0,0 +1,6 @@
+{
+  "info" : {
+    "author" : "xcode",
+    "version" : 1
+  }
+}
diff --git a/Engine/Source/Programs/UnrealMultiUserServer/Resources/Mac/make_icns.sh b/Engine/Source/Programs/UnrealMultiUserServer/Resources/Mac/make_icns.sh
deleted file mode 100644
diff --git a/Engine/Source/Programs/UnrealMultiUserSlateServer/Resources/Mac/Assets.xcassets/AppIcon.appiconset/Contents.json b/Engine/Source/Programs/UnrealMultiUserSlateServer/Resources/Mac/Assets.xcassets/AppIcon.appiconset/Contents.json
new file mode 100644
--- /dev/null
+++ b/Engine/Source/Programs/UnrealMultiUserSlateServer/Resources/Mac/Assets.xcassets/AppIcon.appiconset/Contents.json
@@ -0,0 +1,68 @@
+{
+  "images" : [
+    {
+      "filename" : "icon_16x16.png",
+      "idiom" : "mac",
+      "scale" : "1x",
+      "size" : "16x16"
+    },
+    {
+      "filename" : "[REDACTED_EMAIL]",
+      "idiom" : "mac",
+      "scale" : "2x",
+      "size" : "16x16"
+    },
+    {
+      "filename" : "icon_32x32.png",
+      "idiom" : "mac",
+      "scale" : "1x",
+      "size" : "32x32"
+    },
+    {
+      "filename" : "[REDACTED_EMAIL]",
+      "idiom" : "mac",
+      "scale" : "2x",
+      "size" : "32x32"
+    },
+    {
+      "filename" : "icon_128x128.png",
+      "idiom" : "mac",
+      "scale" : "1x",
+      "size" : "128x128"
+    },
+    {
+      "filename" : "[REDACTED_EMAIL]",
+      "idiom" : "mac",
+      "scale" : "2x",
+      "size" : "128x128"
+    },
+    {
+      "filename" : "icon_256x256.png",
+      "idiom" : "mac",
+      "scale" : "1x",
+      "size" : "256x256"
+    },
+    {
+      "filename" : "[REDACTED_EMAIL]",
+      "idiom" : "mac",
+      "scale" : "2x",
+      "size" : "256x256"
+    },
+    {
+      "filename" : "icon_512x512.png",
+      "idiom" : "mac",
+      "scale" : "1x",
+      "size" : "512x512"
+    },
+    {
+      "filename" : "[REDACTED_EMAIL]",
+      "idiom" : "mac",
+      "scale" : "2x",
+      "size" : "512x512"
+    }
+  ],
+  "info" : {
+    "author" : "xcode",
+    "version" : 1
+  }
+}
diff --git a/Engine/Source/Programs/UnrealMultiUserSlateServer/Resources/Mac/Assets.xcassets/Contents.json b/Engine/Source/Programs/UnrealMultiUserSlateServer/Resources/Mac/Assets.xcassets/Contents.json
new file mode 100644
--- /dev/null
+++ b/Engine/Source/Programs/UnrealMultiUserSlateServer/Resources/Mac/Assets.xcassets/Contents.json
@@ -0,0 +1,6 @@
+{
+  "info" : {
+    "author" : "xcode",
+    "version" : 1
+  }
+}
diff --git a/Engine/Source/Programs/UnrealMultiUserSlateServer/Resources/Mac/make_icns.sh b/Engine/Source/Programs/UnrealMultiUserSlateServer/Resources/Mac/make_icns.sh
deleted file mode 100644
