MultiUser(Slate)Server: Update deprecated EngineIncludeOrderVersion to Latest.

#rnx

diff --git a/Engine/Source/Programs/UnrealMultiUserServer/UnrealMultiUserServer.Target.cs b/Engine/Source/Programs/UnrealMultiUserServer/UnrealMultiUserServer.Target.cs
--- a/Engine/Source/Programs/UnrealMultiUserServer/UnrealMultiUserServer.Target.cs
+++ b/Engine/Source/Programs/UnrealMultiUserServer/UnrealMultiUserServer.Target.cs
@@ -9,6 +9,7 @@ public class UnrealMultiUserServerTarget : TargetRules
 	{
 		Type = TargetType.Program;
 		LinkType = TargetLinkType.Modular;
+		IncludeOrderVersion = EngineIncludeOrderVersion.Latest;
 		LaunchModuleName = "UnrealMultiUserServer";
 		AdditionalPlugins.Add("UdpMessaging");
 		AdditionalPlugins.Add("QuicMessaging");
diff --git a/Engine/Source/Programs/UnrealMultiUserSlateServer/UnrealMultiUserSlateServer.Target.cs b/Engine/Source/Programs/UnrealMultiUserSlateServer/UnrealMultiUserSlateServer.Target.cs
--- a/Engine/Source/Programs/UnrealMultiUserSlateServer/UnrealMultiUserSlateServer.Target.cs
+++ b/Engine/Source/Programs/UnrealMultiUserSlateServer/UnrealMultiUserSlateServer.Target.cs
@@ -9,6 +9,7 @@ public class UnrealMultiUserSlateServerTarget : TargetRules
 	{
 		Type = TargetType.Program;
 		LinkType = TargetLinkType.Modular;
+		IncludeOrderVersion = EngineIncludeOrderVersion.Latest;
 		LaunchModuleName = "UnrealMultiUserSlateServer";
 		AdditionalPlugins.Add("UdpMessaging");
 		AdditionalPlugins.Add("QuicMessaging");
