[macOS/iOS] Fix CIS issue by marking 4 UE programs that depend on the default toolchain as requires platform sdk

#Synced-CL 50951058

diff --git a/Engine/Source/Programs/Mac/DsymExporter/DsymExporter.Build.cs b/Engine/Source/Programs/Mac/DsymExporter/DsymExporter.Build.cs
--- a/Engine/Source/Programs/Mac/DsymExporter/DsymExporter.Build.cs
+++ b/Engine/Source/Programs/Mac/DsymExporter/DsymExporter.Build.cs
@@ -7,6 +7,7 @@ public class DsymExporter : ModuleRules
 	public DsymExporter( ReadOnlyTargetRules Target ) : base(Target)
 	{
 		PrivateIncludePathModuleNames.Add( "Launch" );
+		bRequiresPlatformSDK = true;
 	
 		PrivateDependencyModuleNames.AddRange(
 			new string[] 
diff --git a/Engine/Source/Programs/Mac/UnrealAtoS/UnrealAtoS.Build.cs b/Engine/Source/Programs/Mac/UnrealAtoS/UnrealAtoS.Build.cs
--- a/Engine/Source/Programs/Mac/UnrealAtoS/UnrealAtoS.Build.cs
+++ b/Engine/Source/Programs/Mac/UnrealAtoS/UnrealAtoS.Build.cs
@@ -7,6 +7,7 @@ public class UnrealAtoS : ModuleRules
 	public UnrealAtoS( ReadOnlyTargetRules Target ) : base(Target)
 	{
 		PrivateIncludePathModuleNames.Add( "Launch" );
+		bRequiresPlatformSDK = true;
 	
 		PrivateDependencyModuleNames.AddRange(
 			new string[] 
diff --git a/Engine/Source/Programs/Mac/UnrealEditorServices/UnrealEditorServices.Build.cs b/Engine/Source/Programs/Mac/UnrealEditorServices/UnrealEditorServices.Build.cs
--- a/Engine/Source/Programs/Mac/UnrealEditorServices/UnrealEditorServices.Build.cs
+++ b/Engine/Source/Programs/Mac/UnrealEditorServices/UnrealEditorServices.Build.cs
@@ -7,6 +7,7 @@ public class UnrealEditorServices : ModuleRules
 	public UnrealEditorServices(ReadOnlyTargetRules Target) : base(Target)
 	{
 		PublicIncludePathModuleNames.Add("Launch");
+		bRequiresPlatformSDK = true;
 
 		PrivateDependencyModuleNames.AddRange(
 			new string[] {
diff --git a/Engine/Source/Programs/UnrealFrontend/UnrealFrontend.Build.cs b/Engine/Source/Programs/UnrealFrontend/UnrealFrontend.Build.cs
--- a/Engine/Source/Programs/UnrealFrontend/UnrealFrontend.Build.cs
+++ b/Engine/Source/Programs/UnrealFrontend/UnrealFrontend.Build.cs
@@ -10,6 +10,8 @@ public class UnrealFrontend : ModuleRules
 	{
 		PublicIncludePathModuleNames.Add("Launch");
 
+		bRequiresPlatformSDK = true;
+
 		PrivateIncludePaths.AddRange(
 			new string[] {
 				"Programs/UnrealFrontend/Private/Commands",
diff --git a/Engine/Source/Programs/UnrealMultiUserSlateServer/UnrealMultiUserSlateServer.Build.cs b/Engine/Source/Programs/UnrealMultiUserSlateServer/UnrealMultiUserSlateServer.Build.cs
--- a/Engine/Source/Programs/UnrealMultiUserSlateServer/UnrealMultiUserSlateServer.Build.cs
+++ b/Engine/Source/Programs/UnrealMultiUserSlateServer/UnrealMultiUserSlateServer.Build.cs
@@ -7,6 +7,8 @@ public class UnrealMultiUserSlateServer : ModuleRules
 	public UnrealMultiUserSlateServer(ReadOnlyTargetRules Target) : base(Target)
 	{
 		//ShortName = "MU-SlateServer";
+		bRequiresPlatformSDK = true;
+		
 		PrivateDependencyModuleNames.AddRange(
 			new string[] {
 				"Core",
